Removed unnecessary CoreUObject dependencies

(ushell-p4-cherrypick of 8572493 by [REDACTED_PERSON])

diff --git a/Engine/Source/Developer/TraceServices/TraceServices.Build.cs b/Engine/Source/Developer/TraceServices/TraceServices.Build.cs
--- a/Engine/Source/Developer/TraceServices/TraceServices.Build.cs
+++ b/Engine/Source/Developer/TraceServices/TraceServices.Build.cs
@@ -9,7 +9,6 @@ public class TraceServices : ModuleRules
 		PublicDependencyModuleNames.AddRange(
 			new string[] {
 				"Core",
-				"CoreUObject",
 				"Sockets",
 				"TraceAnalysis",
 			});
diff --git a/Engine/Source/Programs/UnrealInsightsCLI/UnrealInsightsCLI.Target.cs b/Engine/Source/Programs/UnrealInsightsCLI/UnrealInsightsCLI.Target.cs
--- a/Engine/Source/Programs/UnrealInsightsCLI/UnrealInsightsCLI.Target.cs
+++ b/Engine/Source/Programs/UnrealInsightsCLI/UnrealInsightsCLI.Target.cs
@@ -13,7 +13,7 @@ public class UnrealInsightsCLITarget : TargetRules
 		LaunchModuleName = "UnrealInsightsCLI";
 
 		bBuildDeveloperTools = false;
-		bBuildWithEditorOnlyData = true;
+		bBuildWithEditorOnlyData = false;
 		bCompileAgainstEngine = false;
 		bCompileAgainstCoreUObject = false;
 		bCompileAgainstApplicationCore = false;
