Changed UnrealInsightsCLI to allow reports from live connections.

diff --git a/Engine/Source/Developer/TraceServices/Private/Modules/LoadTimeProfilerModule.cpp b/Engine/Source/Developer/TraceServices/Private/Modules/LoadTimeProfilerModule.cpp
--- a/Engine/Source/Developer/TraceServices/Private/Modules/LoadTimeProfilerModule.cpp
+++ b/Engine/Source/Developer/TraceServices/Private/Modules/LoadTimeProfilerModule.cpp
@@ -14,6 +14,19 @@ static const FName LoadTimeProfilerModuleName("TraceModule_LoadTimeProfiler");
 static const FName LoadTimeProfilerProviderName("LoadTimeProfiler");
 static const FName FileActivityProviderName("FileActivity");
 
+void BookmarksToCsv(const IBookmarkProvider& BookmarkProvider, const TCHAR* Filename, double CaptureStartTime, double CaptureEndTime)
+{
+	TSharedPtr<FArchive> BookmarksOutputFile = MakeShareable(IFileManager::Get().CreateFileWriter(Filename));
+	const char* Header = "Name, Timestamp\n";
+	BookmarksOutputFile->Serialize((void*)Header, strlen(Header));
+	BookmarkProvider.EnumerateBookmarks(CaptureStartTime, CaptureEndTime, [&BookmarksOutputFile](const FBookmark& Bookmark)
+	{
+		FString Line = FString::Printf(TEXT("%s,%.4f\n"), Bookmark.Text, Bookmark.Time);
+		auto AnsiLine = StringCast<ANSICHAR>(*Line);
+		BookmarksOutputFile->Serialize((void*)AnsiLine.Get(), AnsiLine.Length());
+	});
+}
+
 void FLoadTimeProfilerModule::GetModuleInfo(FModuleInfo& OutModuleInfo)
 {
 	OutModuleInfo.Name = LoadTimeProfilerModuleName;
@@ -45,6 +58,8 @@ void FLoadTimeProfilerModule::GenerateReports(const IAnalysisSession& Session, c
 	FParse::Value(CmdLine, TEXT("-BeginCaptureBookmark="), BeginCaptureBookmarkName);
 	FString EndCaptureBookmarkName;
 	FParse::Value(CmdLine, TEXT("-EndCaptureBookmark="), EndCaptureBookmarkName);
+	
+
 	BookmarkProvider.EnumerateBookmarks(0.0, DBL_MAX, [BeginCaptureBookmarkName, EndCaptureBookmarkName, &CaptureStartTime, &CaptureEndTime](const FBookmark& Bookmark)
 	{
 		if (CaptureStartTime == -DBL_MAX && !BeginCaptureBookmarkName.IsEmpty() && BeginCaptureBookmarkName == Bookmark.Text)
@@ -64,6 +79,7 @@ void FLoadTimeProfilerModule::GenerateReports(const IAnalysisSession& Session, c
 	{
 		CaptureEndTime = Session.GetDurationSeconds();
 	}
+
 	const ILoadTimeProfilerProvider* LoadTimeProfiler = Trace::ReadLoadTimeProfilerProvider(Session);
 	FString ReportDirectory = FString(OutputDirectory) / TEXT("LoadTimeProfiler");
 	if (LoadTimeProfiler)
@@ -85,6 +101,7 @@ void FLoadTimeProfilerModule::GenerateReports(const IAnalysisSession& Session, c
 		FString SummaryLine = FString::Printf(TEXT("Capture start: %f\r\nCapture end: %f\r\nCapture duration: %f"), CaptureStartTime, CaptureEndTime, CaptureEndTime - CaptureStartTime);
 		CaptureSummaryOutputFile->Serialize(TCHAR_TO_ANSI(*SummaryLine), SummaryLine.Len());
 	}
+	BookmarksToCsv(BookmarkProvider, *(FString(ReportDirectory) / TEXT("Bookmarks.csv")), CaptureStartTime, CaptureEndTime);
 }
 
 const ILoadTimeProfilerProvider* ReadLoadTimeProfilerProvider(const IAnalysisSession& Session)
diff --git a/Engine/Source/Programs/UnrealInsightsCLI/Private/UnrealInsightsCLI.cpp b/Engine/Source/Programs/UnrealInsightsCLI/Private/UnrealInsightsCLI.cpp
--- a/Engine/Source/Programs/UnrealInsightsCLI/Private/UnrealInsightsCLI.cpp
+++ b/Engine/Source/Programs/UnrealInsightsCLI/Private/UnrealInsightsCLI.cpp
@@ -13,45 +13,188 @@ DEFINE_LOG_CATEGORY_STATIC(LogUnrealInsightsCLI, Log, All);
 
 IMPLEMENT_APPLICATION(UnrealInsightsCLI, "UnrealInsightsCLI");
 
-INT32_MAIN_INT32_ARGC_TCHAR_ARGV()
+const char CommandlineHelpText[] = "\n\
+Usage: UnrealInsightsCLI.exe Command [Arguments...]\n\
+\n\
+Commands:\n\
+	ReportFromFile				Generate report from a saved trace file.\n\
+	ReportFromConnection		Generate report by starting analysis server, waiting for\n\
+								the first connection, then running the report on once \n\
+								the session is closed.\n\
+\n\
+Arguments:\n\
+	-inputfile=[file]			Trace file to read from.\n\
+	-outputdir=[dir]			Directory to output from. Default is project 'Report' dir\n\
+\n\
+";
+
+
+/**
+ * Generate the report given a data stream.
+ */
+int GenerateReport(TUniquePtr<Trace::IInDataStream>& DataStream, ITraceServicesModule& TraceServicesModule, const FString& OutputDirectory)
+{
+	using namespace Trace; 
+	
+	TSharedPtr<IAnalysisService> AnalysisService = TraceServicesModule.GetAnalysisService();
+	TSharedPtr<IModuleService> ModuleService = TraceServicesModule.GetModuleService();
+
+	TSharedPtr<const IAnalysisSession> Session = AnalysisService->Analyze(TEXT("Session"), MoveTemp(DataStream));
+
+
+	IPlatformFile& FileSystem = IPlatformFile::GetPlatformPhysical();
+	if (FileSystem.CreateDirectory(*OutputDirectory))
+	{
+		FString OutputDirectoryAbs = FPaths::ConvertRelativePathToFull(OutputDirectory);
+		UE_LOG(LogUnrealInsightsCLI, Display, TEXT("Saving reports to %s"), *OutputDirectoryAbs);
+
+		FAnalysisSessionReadScope _(*Session);
+		ModuleService->GenerateReports(*Session, FCommandLine::Get(), *OutputDirectory);
+	}
+	else
+	{
+		UE_LOG(LogUnrealInsightsCLI, Error, TEXT("Failed to create directory %s"), *OutputDirectory);
+		return 1;
+	}
+
+	return 0;
+}
+
+/**
+ * Generate report from a saved trace file.
+ */
+int Command_ReportFromFile(const TArray<FString> Args)
 {
 	using namespace Trace;
 
-	GEngineLoop.PreInit(ArgC, ArgV);
+	FString FileName;
+	if (!FParse::Value(FCommandLine::Get(), TEXT("-inputfile="), FileName))
+	{
+		UE_LOG(LogUnrealInsightsCLI, Error, TEXT("No file given."));
+		return 1;
+	}
+
+	FString OutputDirectory;
+	if (!FParse::Value(FCommandLine::Get(), TEXT("-outputdir="), OutputDirectory))
+	{
+		// Fall back on project report directory
+		OutputDirectory = FPaths::ProjectSavedDir() / TEXT("Reports") / FDateTime::Now().ToString(TEXT("%Y%m%d_%H%M%S"));
+	}
+
+	if (!FPaths::ValidatePath(OutputDirectory))
+	{
+		UE_LOG(LogUnrealInsightsCLI, Error, TEXT("Outputdirectory %s is not valid path."), *OutputDirectory);
+		return 1;
+	}
+
+	UE_LOG(LogUnrealInsightsCLI, Display, TEXT("Output directory set to %s."), *OutputDirectory);
+
 	ITraceServicesModule& TraceServicesModule = FModuleManager::LoadModuleChecked<ITraceServicesModule>("TraceServices");
-	TSharedPtr<IAnalysisService> AnalysisService = TraceServicesModule.GetAnalysisService();
 	TSharedPtr<Trace::ISessionService> SessionService = TraceServicesModule.GetSessionService();
+	
+	TUniquePtr<IInDataStream> DataStream(SessionService->OpenSessionFromFile(*FileName));
+	if (!DataStream)
+	{
+		UE_LOG(LogUnrealInsightsCLI, Error, TEXT("Trace file %s not found"), *FileName);
+		return 1;
+	}
+	UE_LOG(LogUnrealInsightsCLI, Display, TEXT("Analyzing %s..."), *FileName);
+
+	TSharedPtr<IAnalysisService> AnalysisService = TraceServicesModule.GetAnalysisService();
 	TSharedPtr<IModuleService> ModuleService = TraceServicesModule.GetModuleService();
 
-	TArray<FString> Tokens, Switches;
-	FCommandLine::Parse(FCommandLine::Get(), Tokens, Switches);
-	if (Tokens.Num() < 1)
+	return GenerateReport(DataStream, TraceServicesModule, OutputDirectory);
+}
+
+/**
+ * Generate report by starting analysis server, waiting for the first connection, then running the report
+ * on once the session is closed.
+ */
+int Command_ReportFromConnection(const TArray<FString>&  Args)
+{
+	using namespace Trace;
+
+	FString OutputDirectory;
+	if (!FParse::Value(FCommandLine::Get(), TEXT("-outputdir="), OutputDirectory))
+	{
+		// Fall back on project report directory
+		OutputDirectory = FPaths::ProjectSavedDir() / TEXT("Reports") / FDateTime::Now().ToString(TEXT("%Y%m%d_%H%M%S"));
+	}
+
+	if (!FPaths::ValidatePath(OutputDirectory))
 	{
-		UE_LOG(LogUnrealInsightsCLI, Error, TEXT("Usage: UnrealInsightsCLI.exe TraceFile [Arguments...]"));
+		UE_LOG(LogUnrealInsightsCLI, Error, TEXT("Outputdirectory %s is not valid path."), *OutputDirectory);
 		return 1;
 	}
-	const TCHAR* FileName = *Tokens[0];
-	TUniquePtr<IInDataStream> DataStream(SessionService->OpenSessionFromFile(FileName));
-	if (!DataStream)
+
+	UE_LOG(LogUnrealInsightsCLI, Display, TEXT("Output directory set to %s."), *OutputDirectory);
+
+	ITraceServicesModule& TraceServicesModule = FModuleManager::LoadModuleChecked<ITraceServicesModule>("TraceServices");
+	TSharedPtr<Trace::ISessionService> SessionService = TraceServicesModule.CreateSessionService(*OutputDirectory);
+
+	if (!SessionService->StartRecorderServer())
 	{
-		UE_LOG(LogUnrealInsightsCLI, Error, TEXT("Trace file %s not found"), FileName);
+		UE_LOG(LogUnrealInsightsCLI, Error, TEXT("Failed sto start recording server."));
 		return 1;
 	}
-	UE_LOG(LogUnrealInsightsCLI, Display, TEXT("Analyzing %s..."), FileName);
-	double StartTime = FPlatformTime::Seconds();
-	TSharedPtr<const IAnalysisSession> Session = AnalysisService->Analyze(TEXT("Session"), MoveTemp(DataStream));
-	double Duration = FPlatformTime::Seconds() - StartTime;
-	UE_LOG(LogUnrealInsightsCLI, Display, TEXT("Analysis done! Took %f seconds"), Duration);
 
-	FString OutputDirectory = FPaths::ProjectSavedDir() / TEXT("Reports") / FDateTime::Now().ToString(TEXT("%Y%m%d_%H%M%S"));
-	IPlatformFile& FileSystem = IPlatformFile::GetPlatformPhysical();
-	FileSystem.CreateDirectory(*OutputDirectory);
-	UE_LOG(LogUnrealInsightsCLI, Display, TEXT("Saving reports to %s"), *OutputDirectory);
+	// Wait for the first connection to appear, then wait for it to close.
+	FSessionHandle Session = 0;
+	bool bDoRun = true;
+
+	UE_LOG(LogUnrealInsightsCLI, Display, TEXT("Waiting for connection..."));
+
+	while (bDoRun)
+	{
+		TArray<FSessionHandle> LiveSessions;
+		SessionService->GetLiveSessions(LiveSessions);
+		if (!Session && LiveSessions.Num() > 0)
+		{
+			UE_LOG(LogUnrealInsightsCLI, Display, TEXT("Connection established."));
+			Session = LiveSessions[0];
+		}
+		if (Session && LiveSessions.Num() == 0)
+		{
+			bDoRun = false;
+		}
+		FPlatformProcess::Sleep(0.2f);
+		FTicker::GetCoreTicker().Tick(0.2f);
+	}
+
+	TUniquePtr<IInDataStream> DataStream(SessionService->OpenSessionStream(Session));
+
+	return GenerateReport(DataStream, TraceServicesModule, OutputDirectory);
+}
+
+INT32_MAIN_INT32_ARGC_TCHAR_ARGV()
+{
+	using namespace Trace;
+
+	GEngineLoop.PreInit(ArgC, ArgV);
+
+	int Result = 0;
+	TArray<FString> Tokens, Switches;
+	FCommandLine::Parse(FCommandLine::Get(), Tokens, Switches);
+	if (Tokens.Num() < 1)
+	{
+		printf(CommandlineHelpText);
+		Result = 2;
+	}
+	else
+	{
+		if (Tokens[0].Compare(TEXT("ReportFromFile"), ESearchCase::IgnoreCase) == 0)
+		{
+			Result = Command_ReportFromFile(Switches);
+		}
+		else if (Tokens[0].Compare(TEXT("ReportFromConnection"), ESearchCase::IgnoreCase) == 0)
+		{
+			Result = Command_ReportFromConnection(Switches);
+		}
+	}
 
-	FAnalysisSessionReadScope _(*Session);
-	ModuleService->GenerateReports(*Session, FCommandLine::Get(), *OutputDirectory);
+	FModuleManager::Get().UnloadModulesAtShutdown();
 
 	GIsRequestingExit = true;
-	return 0;
+	return Result;
 }
 
