UnrealInsights:
Removed dependency on SessionService to perform simple offline analysis
Detecting initial enabled/disabled state of toggleable modules by inspecting the command line

diff --git a/Engine/Source/Developer/TraceInsights/Private/Insights/InsightsManager.cpp b/Engine/Source/Developer/TraceInsights/Private/Insights/InsightsManager.cpp
--- a/Engine/Source/Developer/TraceInsights/Private/Insights/InsightsManager.cpp
+++ b/Engine/Source/Developer/TraceInsights/Private/Insights/InsightsManager.cpp
@@ -289,13 +289,9 @@ void FInsightsManager::LoadSession(Trace::FSessionHandle SessionHandle)
 {
 	ResetSession();
 
-	Trace::FSessionInfo SessionInfo;
-	SessionService->GetSessionInfo(SessionHandle, SessionInfo);
-
-	TUniquePtr<Trace::IInDataStream> DataStream(SessionService->OpenSessionStream(SessionHandle));
-	if (DataStream)
+	Session = SessionService->StartAnalysis(SessionHandle);
+	if (Session)
 	{
-		Session = AnalysisService->StartAnalysis(SessionInfo.Name, MoveTemp(DataStream));
 		CurrentSessionHandle = SessionHandle;
 		SpawnAndActivateTabs();
 		OnSessionChanged();
@@ -308,10 +304,9 @@ void FInsightsManager::LoadTraceFile(const FString& TraceFilepath)
 {
 	ResetSession();
 
-	TUniquePtr<Trace::IInDataStream> DataStream(SessionService->OpenSessionFromFile(*TraceFilepath));
-	if (DataStream)
+	Session = AnalysisService->StartAnalysis(*TraceFilepath);
+	if (Session)
 	{
-		Session = AnalysisService->StartAnalysis(*TraceFilepath, MoveTemp(DataStream));
 		CurrentSessionHandle = 0;
 		SpawnAndActivateTabs();
 		OnSessionChanged();
diff --git a/Engine/Source/Developer/TraceServices/Private/AnalysisService.cpp b/Engine/Source/Developer/TraceServices/Private/AnalysisService.cpp
--- a/Engine/Source/Developer/TraceServices/Private/AnalysisService.cpp
+++ b/Engine/Source/Developer/TraceServices/Private/AnalysisService.cpp
@@ -195,11 +195,29 @@ FAnalysisService::~FAnalysisService()
 	}
 }
 
-TSharedPtr<const IAnalysisSession> FAnalysisService::Analyze(const TCHAR* SessionName, TUniquePtr<Trace::IInDataStream>&& InDataStream)
+TSharedPtr<const IAnalysisSession> FAnalysisService::Analyze(const TCHAR* SessionUri)
 {
-	TSharedRef<FAnalysisSession> AnalysisSession = MakeShared<FAnalysisSession>(SessionName);
-	TUniquePtr<Trace::IInDataStream> DataStream = MoveTemp(InDataStream);
-	AnalyzeInternal(AnalysisSession, DataStream.Release());
+	IInDataStream* DataStream = Trace::DataStream_ReadFile(SessionUri);
+	if (!DataStream)
+	{
+		return nullptr;
+	}
+	TSharedRef<FAnalysisSession> AnalysisSession = MakeShared<FAnalysisSession>(SessionUri);
+	AnalyzeInternal(AnalysisSession, DataStream);
+	return AnalysisSession;
+}
+
+TSharedPtr<const IAnalysisSession> FAnalysisService::StartAnalysis(const TCHAR* SessionUri)
+{
+	TUniquePtr<IInDataStream> DataStream(Trace::DataStream_ReadFile(SessionUri));
+	if (!DataStream)
+	{
+		return nullptr;
+	}
+	TSharedRef<FAnalysisSession> AnalysisSession = MakeShared<FAnalysisSession>(SessionUri);
+	TSharedPtr<FAsyncTask<FAnalysisWorker>> Task = MakeShared<FAsyncTask<FAnalysisWorker>>(*this, MoveTemp(DataStream), AnalysisSession);
+	Tasks.Add(Task);
+	Task->StartBackgroundTask();
 	return AnalysisSession;
 }
 
diff --git a/Engine/Source/Developer/TraceServices/Private/AnalysisServicePrivate.h b/Engine/Source/Developer/TraceServices/Private/AnalysisServicePrivate.h
--- a/Engine/Source/Developer/TraceServices/Private/AnalysisServicePrivate.h
+++ b/Engine/Source/Developer/TraceServices/Private/AnalysisServicePrivate.h
@@ -83,7 +83,8 @@ class FAnalysisService
 public:
 	FAnalysisService(FModuleService& ModuleService);
 	virtual ~FAnalysisService();
-	virtual TSharedPtr<const IAnalysisSession> Analyze(const TCHAR* SessionName, TUniquePtr<Trace::IInDataStream>&& DataStream) override;
+	virtual TSharedPtr<const IAnalysisSession> Analyze(const TCHAR* SessionUri) override;
+	virtual TSharedPtr<const IAnalysisSession> StartAnalysis(const TCHAR* SessionUri) override;
 	virtual TSharedPtr<const IAnalysisSession> StartAnalysis(const TCHAR* SessionName, TUniquePtr<Trace::IInDataStream>&& DataStream) override;
 	virtual FAnalysisStartedEvent& OnAnalysisStarted() override { return AnalysisStartedEvent; }
 	virtual FAnalysisFinishedEvent& OnAnalysisFinished() override { return AnalysisFinishedEvent; }
diff --git a/Engine/Source/Developer/TraceServices/Private/ModuleService.cpp b/Engine/Source/Developer/TraceServices/Private/ModuleService.cpp
--- a/Engine/Source/Developer/TraceServices/Private/ModuleService.cpp
+++ b/Engine/Source/Developer/TraceServices/Private/ModuleService.cpp
@@ -85,17 +85,41 @@ void FModuleService::OnAnalysisBegin(IAnalysisSession& Session)
 	}
 }
 
-bool FModuleService::GetModuleLoggers(const FName& ModuleName, TArray<const TCHAR*>& OutLoggers)
+TArray<const TCHAR*> FModuleService::GetModuleLoggers(const FName& ModuleName)
 {
+	TArray<const TCHAR*> Loggers;
+
 	FScopeLock Lock(&CriticalSection);
 	Initialize();
-	IModule** FindIt = ModulesMap.Find(ModuleName);
-	if (!FindIt)
+	IModule* FindIt = ModulesMap.FindRef(ModuleName);
+	if (FindIt)
+	{
+		FindIt->GetLoggers(Loggers);
+	}
+	return Loggers;
+}
+
+TSet<FName> FModuleService::GetEnabledModulesFromCommandLine(const TCHAR* CommandLine)
+{
+	TSet<FName> EnabledModulesFromCommandLine;
+
+	if (!CommandLine)
 	{
-		return false;
+		return EnabledModulesFromCommandLine;
+	}
+
+	FScopeLock Lock(&CriticalSection);
+	Initialize();
+	for (const auto& KV : ModulesMap)
+	{
+		IModule* Module = KV.Value;
+		const TCHAR* ModuleCommandLineArgument = Module->GetCommandLineArgument();
+		if (ModuleCommandLineArgument && FParse::Param(CommandLine, ModuleCommandLineArgument))
+		{
+			EnabledModulesFromCommandLine.Add(KV.Key);
+		}
 	}
-	(*FindIt)->GetLoggers(OutLoggers);
-	return true;
+	return EnabledModulesFromCommandLine;
 }
 
 void FModuleService::GenerateReports(const IAnalysisSession& Session, const TCHAR* CmdLine, const TCHAR* OutputDirectory)
diff --git a/Engine/Source/Developer/TraceServices/Private/ModuleServicePrivate.h b/Engine/Source/Developer/TraceServices/Private/ModuleServicePrivate.h
--- a/Engine/Source/Developer/TraceServices/Private/ModuleServicePrivate.h
+++ b/Engine/Source/Developer/TraceServices/Private/ModuleServicePrivate.h
@@ -21,7 +21,8 @@ public:
 	virtual void SetModuleEnabled(const FName& ModuleName, bool bEnabled) override;
 	virtual void GenerateReports(const IAnalysisSession& Session, const TCHAR* CmdLine, const TCHAR* OutputDirectory) override;
 	void OnAnalysisBegin(IAnalysisSession& Session);
-	bool GetModuleLoggers(const FName& ModuleName, TArray<const TCHAR*>& OutLoggers);
+	TArray<const TCHAR*> GetModuleLoggers(const FName& ModuleName);
+	TSet<FName> GetEnabledModulesFromCommandLine(const TCHAR* CommandLine);
 
 private:
 	void Initialize();
diff --git a/Engine/Source/Developer/TraceServices/Private/Modules/CountersModule.h b/Engine/Source/Developer/TraceServices/Private/Modules/CountersModule.h
--- a/Engine/Source/Developer/TraceServices/Private/Modules/CountersModule.h
+++ b/Engine/Source/Developer/TraceServices/Private/Modules/CountersModule.h
@@ -14,6 +14,10 @@ public:
 	virtual void GetModuleInfo(FModuleInfo& OutModuleInfo) override;
 	virtual void OnAnalysisBegin(IAnalysisSession& Session) override;
 	virtual void GetLoggers(TArray<const TCHAR *>& OutLoggers) override;
+	virtual const TCHAR* GetCommandLineArgument() override
+	{
+		return TEXT("counterstrace");
+	}
 	virtual void GenerateReports(const IAnalysisSession& Session, const TCHAR* CmdLine, const TCHAR* OutputDirectory) override {}
 };
 
diff --git a/Engine/Source/Developer/TraceServices/Private/Modules/CsvProfilerModule.h b/Engine/Source/Developer/TraceServices/Private/Modules/CsvProfilerModule.h
--- a/Engine/Source/Developer/TraceServices/Private/Modules/CsvProfilerModule.h
+++ b/Engine/Source/Developer/TraceServices/Private/Modules/CsvProfilerModule.h
@@ -14,7 +14,13 @@ class FCsvProfilerModule
 public:
 	virtual void GetModuleInfo(FModuleInfo& OutModuleInfo) override;
 	virtual void OnAnalysisBegin(IAnalysisSession& Session) override;
-	virtual void GetLoggers(TArray<const TCHAR*>& OutLoggers) override {};
+	virtual void GetLoggers(TArray<const TCHAR*>& OutLoggers) override
+	{
+	};
+	virtual const TCHAR* GetCommandLineArgument() override
+	{
+		return nullptr;
+	}
 	virtual void GenerateReports(const IAnalysisSession& Session, const TCHAR* CmdLine, const TCHAR* OutputDirectory) override;
 };
 
diff --git a/Engine/Source/Developer/TraceServices/Private/Modules/LoadTimeProfilerModule.h b/Engine/Source/Developer/TraceServices/Private/Modules/LoadTimeProfilerModule.h
--- a/Engine/Source/Developer/TraceServices/Private/Modules/LoadTimeProfilerModule.h
+++ b/Engine/Source/Developer/TraceServices/Private/Modules/LoadTimeProfilerModule.h
@@ -14,6 +14,10 @@ public:
 	virtual void GetModuleInfo(FModuleInfo& OutModuleInfo) override;
 	virtual void OnAnalysisBegin(IAnalysisSession& Session) override;
 	virtual void GetLoggers(TArray<const TCHAR *>& OutLoggers) override;
+	virtual const TCHAR* GetCommandLineArgument() override
+	{
+		return TEXT("loadtimetrace");
+	}
 	virtual void GenerateReports(const IAnalysisSession& Session, const TCHAR* CmdLine, const TCHAR* OutputDirectory) override;
 };
 
diff --git a/Engine/Source/Developer/TraceServices/Private/Modules/NetProfilerModule.cpp b/Engine/Source/Developer/TraceServices/Private/Modules/NetProfilerModule.cpp
--- a/Engine/Source/Developer/TraceServices/Private/Modules/NetProfilerModule.cpp
+++ b/Engine/Source/Developer/TraceServices/Private/Modules/NetProfilerModule.cpp
@@ -28,7 +28,7 @@ void FNetProfilerModule::OnAnalysisBegin(IAnalysisSession& InSession)
 
 void FNetProfilerModule::GetLoggers(TArray<const TCHAR *>& OutLoggers)
 {
-	OutLoggers.Add(TEXT("NetProfiler"));
+	//OutLoggers.Add(TEXT("NetProfiler"));
 }
 
 }
diff --git a/Engine/Source/Developer/TraceServices/Private/Modules/NetProfilerModule.h b/Engine/Source/Developer/TraceServices/Private/Modules/NetProfilerModule.h
--- a/Engine/Source/Developer/TraceServices/Private/Modules/NetProfilerModule.h
+++ b/Engine/Source/Developer/TraceServices/Private/Modules/NetProfilerModule.h
@@ -15,6 +15,10 @@ public:
 	virtual void GetModuleInfo(FModuleInfo& OutModuleInfo) override;
 	virtual void OnAnalysisBegin(IAnalysisSession& Session) override;
 	virtual void GetLoggers(TArray<const TCHAR*>& OutLoggers) override;
+	virtual const TCHAR* GetCommandLineArgument() override
+	{
+		return nullptr;
+	}
 	virtual void GenerateReports(const IAnalysisSession& Session, const TCHAR* CmdLine, const TCHAR* OutputDirectory) override {}
 
 private:
diff --git a/Engine/Source/Developer/TraceServices/Private/Modules/StatsModule.cpp b/Engine/Source/Developer/TraceServices/Private/Modules/StatsModule.cpp
--- a/Engine/Source/Developer/TraceServices/Private/Modules/StatsModule.cpp
+++ b/Engine/Source/Developer/TraceServices/Private/Modules/StatsModule.cpp
@@ -24,7 +24,7 @@ void FStatsModule::OnAnalysisBegin(IAnalysisSession& Session)
 
 void FStatsModule::GetLoggers(TArray<const TCHAR *>& OutLoggers)
 {
-	OutLoggers.Add(TEXT("Stats"));
+	//OutLoggers.Add(TEXT("Stats"));
 }
 
 }
diff --git a/Engine/Source/Developer/TraceServices/Private/Modules/StatsModule.h b/Engine/Source/Developer/TraceServices/Private/Modules/StatsModule.h
--- a/Engine/Source/Developer/TraceServices/Private/Modules/StatsModule.h
+++ b/Engine/Source/Developer/TraceServices/Private/Modules/StatsModule.h
@@ -15,6 +15,10 @@ public:
 	virtual void GetModuleInfo(FModuleInfo& OutModuleInfo) override;
 	virtual void OnAnalysisBegin(IAnalysisSession& Session) override;
 	virtual void GetLoggers(TArray<const TCHAR *>& OutLoggers) override;
+	virtual const TCHAR* GetCommandLineArgument() override
+	{
+		return nullptr;
+	}
 	virtual void GenerateReports(const IAnalysisSession& Session, const TCHAR* CmdLine, const TCHAR* OutputDirectory) override {}
 };
 
diff --git a/Engine/Source/Developer/TraceServices/Private/Modules/TimingProfilerModule.h b/Engine/Source/Developer/TraceServices/Private/Modules/TimingProfilerModule.h
--- a/Engine/Source/Developer/TraceServices/Private/Modules/TimingProfilerModule.h
+++ b/Engine/Source/Developer/TraceServices/Private/Modules/TimingProfilerModule.h
@@ -15,6 +15,10 @@ public:
 	virtual void GetModuleInfo(FModuleInfo& OutModuleInfo) override;
 	virtual void OnAnalysisBegin(IAnalysisSession& Session) override;
 	virtual void GetLoggers(TArray<const TCHAR *>& OutLoggers) override;
+	virtual const TCHAR* GetCommandLineArgument() override
+	{
+		return TEXT("cpuprofilertrace");
+	}
 	virtual void GenerateReports(const IAnalysisSession& Session, const TCHAR* CmdLine, const TCHAR* OutputDirectory) override {}
 };
 
diff --git a/Engine/Source/Developer/TraceServices/Private/SessionService.cpp b/Engine/Source/Developer/TraceServices/Private/SessionService.cpp
--- a/Engine/Source/Developer/TraceServices/Private/SessionService.cpp
+++ b/Engine/Source/Developer/TraceServices/Private/SessionService.cpp
@@ -2,6 +2,8 @@
 
 #include "TraceServices/SessionService.h"
 #include "SessionServicePrivate.h"
+#include "ModuleServicePrivate.h"
+#include "AnalysisServicePrivate.h"
 #include "Trace/DataStream.h"
 #include "Misc/Paths.h"
 #include "Misc/OutputDeviceRedirector.h"
@@ -64,14 +66,15 @@ struct FDiagnosticsSessionAnalyzer
 	int8 TargetType = 0;
 };
 
-FSessionService::FSessionService(FModuleService& InModuleService)
-	: FSessionService(InModuleService, nullptr)
+FSessionService::FSessionService(FModuleService& InModuleService, FAnalysisService& InAnalysisService)
+	: FSessionService(InModuleService, InAnalysisService, nullptr)
 {
 
 }
 
-FSessionService::FSessionService(FModuleService& InModuleService, const TCHAR* OverrideSessionDirectory)
+FSessionService::FSessionService(FModuleService& InModuleService, FAnalysisService& InAnalysisService, const TCHAR* OverrideSessionDirectory)
 	: ModuleService(InModuleService)
+	, AnalysisService(InAnalysisService)
 {
 	if (OverrideSessionDirectory)
 	{
@@ -172,16 +175,6 @@ bool FSessionService::GetSessionInfo(FSessionHandle SessionHandle, FSessionInfo&
 	return true;
 }
 
-Trace::IInDataStream* FSessionService::OpenSessionStream(FSessionHandle SessionHandle)
-{
-	return TraceStore->OpenSessionStream(SessionHandle);
-}
-
-Trace::IInDataStream* FSessionService::OpenSessionFromFile(const TCHAR* FilePath)
-{
-	return Trace::DataStream_ReadFile(FilePath);
-}
-
 void FSessionService::SetModuleEnabled(FSessionHandle SessionHandle, const FName& ModuleName, bool bState)
 {
 	FScopeLock Lock(&SessionsCS);
@@ -192,11 +185,7 @@ void FSessionService::SetModuleEnabled(FSessionHandle SessionHandle, const FName
 	}
 	if (bState)
 	{
-		TArray<const TCHAR*> Loggers;
-		if (!ModuleService.GetModuleLoggers(ModuleName, Loggers))
-		{
-			return;
-		}
+		TArray<const TCHAR*> Loggers = ModuleService.GetModuleLoggers(ModuleName);
 		TSet<FString>& EnabledLoggers = FindIt->EnabledModuleLoggersMap.FindOrAdd(ModuleName);
 		for (const TCHAR* Logger : Loggers)
 		{
@@ -285,6 +274,38 @@ bool FSessionService::ConnectSession(const TCHAR* ControlClientAddress)
 	return true;
 }
 
+TSharedPtr<const IAnalysisSession> FSessionService::StartAnalysis(FSessionHandle SessionHandle)
+{
+	FScopeLock Lock(&SessionsCS);
+
+	FSessionInfoInternal* FindIt = Sessions.Find(SessionHandle);
+	if (!FindIt)
+	{
+		return nullptr;
+	}
+	
+	TUniquePtr<IInDataStream> DataStream(TraceStore->OpenSessionStream(SessionHandle));
+	if (!DataStream)
+	{
+		return nullptr;
+	}
+
+	if (!FindIt->CommandLine.IsEmpty())
+	{
+		TSet<FName> CommandLineEnabledModules = ModuleService.GetEnabledModulesFromCommandLine(*FindIt->CommandLine);
+		for (const FName& EnabledModuleName : CommandLineEnabledModules)
+		{
+			TSet<FString>& EnabledLoggers = FindIt->EnabledModuleLoggersMap.FindOrAdd(EnabledModuleName);
+			for (const TCHAR* Logger : ModuleService.GetModuleLoggers(EnabledModuleName))
+			{
+				EnabledLoggers.Add(Logger);
+			}
+		}
+	}
+
+	return AnalysisService.StartAnalysis(FindIt->Name, MoveTemp(DataStream));
+}
+
 bool FSessionService::Tick(float DeltaTime)
 {
 	UpdateSessions();
diff --git a/Engine/Source/Developer/TraceServices/Private/SessionServicePrivate.h b/Engine/Source/Developer/TraceServices/Private/SessionServicePrivate.h
--- a/Engine/Source/Developer/TraceServices/Private/SessionServicePrivate.h
+++ b/Engine/Source/Developer/TraceServices/Private/SessionServicePrivate.h
@@ -7,17 +7,19 @@
 #include "Containers/Ticker.h"
 #include "Trace/Recorder.h"
 #include "Trace/Store.h"
-#include "ModuleServicePrivate.h"
 
 namespace Trace
 {
 
+class FModuleService;
+class FAnalysisService;
+
 class FSessionService
 	: public ISessionService
 {
 public:
-	FSessionService(FModuleService& ModuleService);
-	FSessionService(FModuleService& ModuleService, const TCHAR* OverrideSessionsDirectory);
+	FSessionService(FModuleService& ModuleService, FAnalysisService& AnalysisService);
+	FSessionService(FModuleService& ModuleService, FAnalysisService& AnalysisService, const TCHAR* OverrideSessionsDirectory);
 	virtual ~FSessionService();
 	virtual bool StartRecorderServer() override;
 	virtual bool IsRecorderServerRunning() const override;
@@ -26,11 +28,10 @@ public:
 	virtual void GetAvailableSessions(TArray<Trace::FSessionHandle>& OutSessions) const override;
 	virtual void GetLiveSessions(TArray<Trace::FSessionHandle>& OutSessions) const override;
 	virtual bool GetSessionInfo(Trace::FSessionHandle SessionHandle, Trace::FSessionInfo& OutSessionInfo) const override;
-	virtual Trace::IInDataStream* OpenSessionStream(Trace::FSessionHandle SessionHandle) override;
-	virtual Trace::IInDataStream* OpenSessionFromFile(const TCHAR* FilePath) override;
 	virtual void SetModuleEnabled(Trace::FSessionHandle SessionHandle, const FName& ModuleName, bool bState) override;
 	virtual bool IsModuleEnabled(Trace::FSessionHandle SessionHandle, const FName& ModuleName) const override;
 	virtual bool ConnectSession(const TCHAR* ControlClientAddress) override;
+	virtual TSharedPtr<const IAnalysisSession> StartAnalysis(FSessionHandle SessionHandle) override;
 
 private:
 	struct FSessionInfoInternal
@@ -55,6 +56,7 @@ private:
 	void UpdateSessionContext(FStoreSessionHandle StoreHandle, FSessionInfoInternal& Info);
 
 	FModuleService& ModuleService;
+	FAnalysisService& AnalysisService;
 	void* RecorderEvent = nullptr;
 	FString LocalSessionDirectory;
 	TSharedPtr<IStore> TraceStore;
diff --git a/Engine/Source/Developer/TraceServices/Private/TraceServicesModule.cpp b/Engine/Source/Developer/TraceServices/Private/TraceServicesModule.cpp
--- a/Engine/Source/Developer/TraceServices/Private/TraceServicesModule.cpp
+++ b/Engine/Source/Developer/TraceServices/Private/TraceServicesModule.cpp
@@ -45,7 +45,7 @@ TSharedPtr<Trace::ISessionService> FTraceServicesModule::GetSessionService()
 	if (!SessionService.IsValid())
 	{
 		GetModuleService();
-		SessionService = MakeShared<Trace::FSessionService>(*ModuleService.Get());
+		SessionService = MakeShared<Trace::FSessionService>(*ModuleService.Get(), *AnalysisService.Get());
 	}
 	return SessionService;
 }
@@ -73,7 +73,7 @@ TSharedPtr<Trace::ISessionService> FTraceServicesModule::CreateSessionService(co
 {
 	checkf(!SessionService.IsValid(), TEXT("A SessionService already exists."));
 	GetModuleService();
-	SessionService = MakeShared<Trace::FSessionService>(*ModuleService.Get(), SessionDirectory);
+	SessionService = MakeShared<Trace::FSessionService>(*ModuleService.Get(), *AnalysisService.Get(), SessionDirectory);
 	return SessionService;
 }
 
diff --git a/Engine/Source/Developer/TraceServices/Public/TraceServices/AnalysisService.h b/Engine/Source/Developer/TraceServices/Public/TraceServices/AnalysisService.h
--- a/Engine/Source/Developer/TraceServices/Public/TraceServices/AnalysisService.h
+++ b/Engine/Source/Developer/TraceServices/Public/TraceServices/AnalysisService.h
@@ -22,7 +22,8 @@ namespace Trace
 class IAnalysisService
 {
 public:
-	virtual TSharedPtr<const IAnalysisSession> Analyze(const TCHAR* SessionName, TUniquePtr<Trace::IInDataStream>&& DataStream) = 0;
+	virtual TSharedPtr<const IAnalysisSession> Analyze(const TCHAR* SessionUri) = 0;
+	virtual TSharedPtr<const IAnalysisSession> StartAnalysis(const TCHAR* SessionUri) = 0;
 	virtual TSharedPtr<const IAnalysisSession> StartAnalysis(const TCHAR* SessionName, TUniquePtr<Trace::IInDataStream>&& DataStream) = 0;
 
 	DECLARE_EVENT_OneParam(IAnalysisService, FAnalysisStartedEvent, TSharedRef<const IAnalysisSession>)
diff --git a/Engine/Source/Developer/TraceServices/Public/TraceServices/ModuleService.h b/Engine/Source/Developer/TraceServices/Public/TraceServices/ModuleService.h
--- a/Engine/Source/Developer/TraceServices/Public/TraceServices/ModuleService.h
+++ b/Engine/Source/Developer/TraceServices/Public/TraceServices/ModuleService.h
@@ -28,6 +28,7 @@ public:
 	virtual void OnAnalysisBegin(IAnalysisSession& Session) = 0;
 	virtual void GetLoggers(TArray<const TCHAR*>& OutLoggers) = 0;
 	virtual void GenerateReports(const IAnalysisSession& Session, const TCHAR* CmdLine, const TCHAR* OutputDirectory) = 0;
+	virtual const TCHAR* GetCommandLineArgument() = 0;
 };
 
 class IModuleService
diff --git a/Engine/Source/Developer/TraceServices/Public/TraceServices/SessionService.h b/Engine/Source/Developer/TraceServices/Public/TraceServices/SessionService.h
--- a/Engine/Source/Developer/TraceServices/Public/TraceServices/SessionService.h
+++ b/Engine/Source/Developer/TraceServices/Public/TraceServices/SessionService.h
@@ -12,6 +12,7 @@ namespace Trace
 {
 
 typedef uint64 FSessionHandle;
+class IAnalysisSession;
 
 struct FSessionInfo
 {
@@ -38,11 +39,10 @@ public:
 	virtual void GetAvailableSessions(TArray<FSessionHandle>& OutSessions) const = 0;
 	virtual void GetLiveSessions(TArray<FSessionHandle>& OutSessions) const = 0;
 	virtual bool GetSessionInfo(FSessionHandle SessionHandle, FSessionInfo& OutSessionInfo) const = 0;
-	virtual Trace::IInDataStream* OpenSessionStream(FSessionHandle SessionHandle) = 0;
-	virtual Trace::IInDataStream* OpenSessionFromFile(const TCHAR* FilePath) = 0;
 	virtual void SetModuleEnabled(Trace::FSessionHandle SessionHandle, const FName& ModuleName, bool bState) = 0;
 	virtual bool IsModuleEnabled(Trace::FSessionHandle SessionHandle, const FName& ModuleName) const = 0;
 	virtual bool ConnectSession(const TCHAR* ControlClientAddress) = 0;
+	virtual TSharedPtr<const IAnalysisSession> StartAnalysis(FSessionHandle SessionHandle) = 0;
 };
 
 }
diff --git a/Engine/Source/Programs/UnrealInsightsCLI/Private/UnrealInsightsCLI.cpp b/Engine/Source/Programs/UnrealInsightsCLI/Private/UnrealInsightsCLI.cpp
--- a/Engine/Source/Programs/UnrealInsightsCLI/Private/UnrealInsightsCLI.cpp
+++ b/Engine/Source/Programs/UnrealInsightsCLI/Private/UnrealInsightsCLI.cpp
@@ -30,17 +30,18 @@ Arguments:\n\
 
 
 /**
- * Generate the report given a data stream.
+ * Generate the report given a session URI.
  */
-int GenerateReport(TUniquePtr<Trace::IInDataStream>& DataStream, ITraceServicesModule& TraceServicesModule, const FString& OutputDirectory)
+int GenerateReport(const TCHAR* SessionUri, Trace::IAnalysisService& AnalysisService, Trace::IModuleService& ModuleService, const FString& OutputDirectory)
 {
 	using namespace Trace; 
 	
-	TSharedPtr<IAnalysisService> AnalysisService = TraceServicesModule.GetAnalysisService();
-	TSharedPtr<IModuleService> ModuleService = TraceServicesModule.GetModuleService();
-
-	TSharedPtr<const IAnalysisSession> Session = AnalysisService->Analyze(TEXT("Session"), MoveTemp(DataStream));
-
+	TSharedPtr<const IAnalysisSession> Session = AnalysisService.Analyze(SessionUri);
+	if (!Session)
+	{
+		UE_LOG(LogUnrealInsightsCLI, Error, TEXT("Trace file %s not found"), SessionUri);
+		return 1;
+	}
 
 	IPlatformFile& FileSystem = IPlatformFile::GetPlatformPhysical();
 	if (FileSystem.CreateDirectory(*OutputDirectory))
@@ -49,7 +50,7 @@ int GenerateReport(TUniquePtr<Trace::IInDataStream>& DataStream, ITraceServicesM
 		UE_LOG(LogUnrealInsightsCLI, Display, TEXT("Saving reports to %s"), *OutputDirectoryAbs);
 
 		FAnalysisSessionReadScope _(*Session);
-		ModuleService->GenerateReports(*Session, FCommandLine::Get(), *OutputDirectory);
+		ModuleService.GenerateReports(*Session, FCommandLine::Get(), *OutputDirectory);
 	}
 	else
 	{
@@ -90,20 +91,9 @@ int Command_ReportFromFile(const TArray<FString> Args)
 	UE_LOG(LogUnrealInsightsCLI, Display, TEXT("Output directory set to %s."), *OutputDirectory);
 
 	ITraceServicesModule& TraceServicesModule = FModuleManager::LoadModuleChecked<ITraceServicesModule>("TraceServices");
-	TSharedPtr<Trace::ISessionService> SessionService = TraceServicesModule.GetSessionService();
 	
-	TUniquePtr<IInDataStream> DataStream(SessionService->OpenSessionFromFile(*FileName));
-	if (!DataStream)
-	{
-		UE_LOG(LogUnrealInsightsCLI, Error, TEXT("Trace file %s not found"), *FileName);
-		return 1;
-	}
 	UE_LOG(LogUnrealInsightsCLI, Display, TEXT("Analyzing %s..."), *FileName);
-
-	TSharedPtr<IAnalysisService> AnalysisService = TraceServicesModule.GetAnalysisService();
-	TSharedPtr<IModuleService> ModuleService = TraceServicesModule.GetModuleService();
-
-	return GenerateReport(DataStream, TraceServicesModule, OutputDirectory);
+	return GenerateReport(*FileName, *TraceServicesModule.GetAnalysisService(), *TraceServicesModule.GetModuleService(), OutputDirectory);
 }
 
 /**
@@ -161,9 +151,14 @@ int Command_ReportFromConnection(const TArray<FString>&  Args)
 		FTicker::GetCoreTicker().Tick(0.2f);
 	}
 
-	TUniquePtr<IInDataStream> DataStream(SessionService->OpenSessionStream(Session));
+	Trace::FSessionInfo SessionInfo;
+	if (!SessionService->GetSessionInfo(Session, SessionInfo))
+	{
+		UE_LOG(LogUnrealInsightsCLI, Error, TEXT("Failed to load session"));
+		return 1;
+	}
 
-	return GenerateReport(DataStream, TraceServicesModule, OutputDirectory);
+	return GenerateReport(SessionInfo.Uri, *TraceServicesModule.GetAnalysisService(), *TraceServicesModule.GetModuleService(), OutputDirectory);
 }
 
 INT32_MAIN_INT32_ARGC_TCHAR_ARGV()
@@ -194,7 +189,7 @@ INT32_MAIN_INT32_ARGC_TCHAR_ARGV()
 
 	FModuleManager::Get().UnloadModulesAtShutdown();
 
-	GIsRequestingExit = true;
+	RequestEngineExit(TEXT("UnrealInsightsCLI finished"));
 	return Result;
 }
 
